Pose Correctives: Fix Mac non-unity builds.

#rnx

diff --git a/Engine/Plugins/Experimental/Animation/PoseCorrectives/Source/PoseCorrectivesEditor/Public/PoseCorrectivesEditorController.h b/Engine/Plugins/Experimental/Animation/PoseCorrectives/Source/PoseCorrectivesEditor/Public/PoseCorrectivesEditorController.h
--- a/Engine/Plugins/Experimental/Animation/PoseCorrectives/Source/PoseCorrectivesEditor/Public/PoseCorrectivesEditorController.h
+++ b/Engine/Plugins/Experimental/Animation/PoseCorrectives/Source/PoseCorrectivesEditor/Public/PoseCorrectivesEditorController.h
@@ -10,12 +10,17 @@
 #include "Widgets/SWindow.h"
 #include "Rigs/RigHierarchyDefines.h"
 
-class UPoseCorrectivesAsset;
+
+class FCorrectivesEditMode;
 class FPoseCorrectivesEditorToolkit;
-class UPoseCorrectivesAnimInstance;
+class IDetailsView;
+class UControlRig;
 class UControlRigComponent;
 class UControlRigSkeletalMeshComponent;
-class FCorrectivesEditMode;
+class UPoseCorrectivesAnimInstance;
+class UPoseCorrectivesAnimSourceInstance;
+class UPoseCorrectivesAsset;
+
 
 class FPoseCorrectivesEditorController : public TSharedFromThis<FPoseCorrectivesEditorController>, FGCObject
 {
@@ -41,14 +46,11 @@ public:
 	UDebugSkelMeshComponent* TargetSkelMeshComponent;
 
 	/** viewport anim instance */
-	UPROPERTY(transient, NonTransactional)
-	TWeakObjectPtr<class UPoseCorrectivesAnimSourceInstance> SourceAnimInstance;
-	
-	UPROPERTY(transient, NonTransactional)
-	TWeakObjectPtr<class UPoseCorrectivesAnimInstance> TargetAnimInstance;	
+	TWeakObjectPtr<UPoseCorrectivesAnimSourceInstance> SourceAnimInstance;
+	TWeakObjectPtr<UPoseCorrectivesAnimInstance> TargetAnimInstance;	
 	
 	/** asset properties tab */
-	TSharedPtr<class IDetailsView> DetailsView;
+	TSharedPtr<IDetailsView> DetailsView;
 
 	/** get the source skeletal mesh we are copying FROM */
 	USkeletalMesh* GetSourceSkeletalMesh() const;
@@ -103,6 +105,5 @@ private:
 	TSharedPtr<SCorrectivesViewer> CorrectivesViewer;
 	
 	/* Editor Control rig*/
-	UPROPERTY()
 	UControlRig* ControlRig = nullptr;
 };
