- Fix crash on skeletal mesh load in PoseCorrectives Editor
- Fix null check on posecorrectivesprocessor

diff --git a/Engine/Plugins/Experimental/Animation/PoseCorrectives/Source/PoseCorrectives/Private/AnimNode_CorrectPose.cpp b/Engine/Plugins/Experimental/Animation/PoseCorrectives/Source/PoseCorrectives/Private/AnimNode_CorrectPose.cpp
--- a/Engine/Plugins/Experimental/Animation/PoseCorrectives/Source/PoseCorrectives/Private/AnimNode_CorrectPose.cpp
+++ b/Engine/Plugins/Experimental/Animation/PoseCorrectives/Source/PoseCorrectives/Private/AnimNode_CorrectPose.cpp
@@ -16,7 +16,7 @@ void FAnimNode_CorrectPose::Initialize_AnyThread(const FAnimationInitializeConte
 
 	SourcePose.Initialize(Context);
 
-	if (PoseCorrectivesProcessor && IsInGameThread())
+	if (!PoseCorrectivesProcessor && IsInGameThread())
 	{
 		PoseCorrectivesProcessor = NewObject<UPoseCorrectivesProcessor>();	
 	}
@@ -107,7 +107,7 @@ void FAnimNode_CorrectPose::Evaluate_AnyThread(FPoseContext& Output)
 	SourcePose.Evaluate(SourceData);
 	Output = SourceData;
 
-	if (!IsValid(PoseCorrectivesAsset) || EditMode)
+	if (!IsValid(PoseCorrectivesAsset) || EditMode || !PoseCorrectivesProcessor)
 	{
 		return;
 	}
diff --git a/Engine/Plugins/Experimental/Animation/PoseCorrectives/Source/PoseCorrectivesEditor/Private/PoseCorrectivesEditorController.cpp b/Engine/Plugins/Experimental/Animation/PoseCorrectives/Source/PoseCorrectivesEditor/Private/PoseCorrectivesEditorController.cpp
--- a/Engine/Plugins/Experimental/Animation/PoseCorrectives/Source/PoseCorrectivesEditor/Private/PoseCorrectivesEditorController.cpp
+++ b/Engine/Plugins/Experimental/Animation/PoseCorrectives/Source/PoseCorrectivesEditor/Private/PoseCorrectivesEditorController.cpp
@@ -95,6 +95,14 @@ void FPoseCorrectivesEditorController::UninitializeTargetControlRigBP()
 	}
 }
 
+void FPoseCorrectivesEditorController::AddReferencedObjects(FReferenceCollector& Collector)
+{
+	Collector.AddReferencedObject(ControlRig);
+	Collector.AddReferencedObject(SourceAnimInstance);
+	Collector.AddReferencedObject(TargetAnimInstance);
+};
+
+
 USkeletalMesh* FPoseCorrectivesEditorController::GetSourceSkeletalMesh() const
 {
 	if (!Asset)
@@ -117,7 +125,7 @@ USkeletalMesh* FPoseCorrectivesEditorController::GetTargetSkeletalMesh() const
 
 void FPoseCorrectivesEditorController::PlayAnimationAsset(UAnimationAsset* AssetToPlay)
 {
-	if (AssetToPlay && SourceAnimInstance.IsValid())
+	if (AssetToPlay && SourceAnimInstance.Get())
 	{
 		ExitEditMode();
 
diff --git a/Engine/Plugins/Experimental/Animation/PoseCorrectives/Source/PoseCorrectivesEditor/Public/PoseCorrectivesEditorController.h b/Engine/Plugins/Experimental/Animation/PoseCorrectives/Source/PoseCorrectivesEditor/Public/PoseCorrectivesEditorController.h
--- a/Engine/Plugins/Experimental/Animation/PoseCorrectives/Source/PoseCorrectivesEditor/Public/PoseCorrectivesEditorController.h
+++ b/Engine/Plugins/Experimental/Animation/PoseCorrectives/Source/PoseCorrectivesEditor/Public/PoseCorrectivesEditorController.h
@@ -46,8 +46,8 @@ public:
 	UDebugSkelMeshComponent* TargetSkelMeshComponent;
 
 	/** viewport anim instance */
-	TWeakObjectPtr<UPoseCorrectivesAnimSourceInstance> SourceAnimInstance;
-	TWeakObjectPtr<UPoseCorrectivesAnimInstance> TargetAnimInstance;	
+	TObjectPtr<UPoseCorrectivesAnimSourceInstance> SourceAnimInstance;
+	TObjectPtr<UPoseCorrectivesAnimInstance> TargetAnimInstance;	
 	
 	/** asset properties tab */
 	TSharedPtr<IDetailsView> DetailsView;
@@ -63,10 +63,7 @@ public:
 	UAnimationAsset* PreviousAsset = nullptr;
 	/** END Sequence Browser */
 
-	virtual void AddReferencedObjects(FReferenceCollector& Collector) override
-	{
-		Collector.AddReferencedObject(ControlRig);
-	};
+	virtual void AddReferencedObjects(FReferenceCollector& Collector) override;
 	virtual FString GetReferencerName() const override {return "PoseCorrectivesEditorController";};
 
 	void HandleEditCorrective(const FName& CorrectiveName);
